Looks related to CL-6597894.  It added IsValidLongPackageName check in FIntroTutorials::HandleSDKNotInstalled, which is failing, but skipping over that loads the tutorial fine.  The safer fix is to just specify the asset name as .../SetupLuminTutorial and not .../SetupLuminTutorial.SetupLuminTutorial, because it was choking on the period.

diff --git a/Engine/Config/Lumin/DataDrivenPlatformInfo.ini b/Engine/Config/Lumin/DataDrivenPlatformInfo.ini
--- a/Engine/Config/Lumin/DataDrivenPlatformInfo.ini
+++ b/Engine/Config/Lumin/DataDrivenPlatformInfo.ini
@@ -13,7 +13,7 @@ LargeIconPath=Launcher/Lumin/Platform_Lumin_128x
 XLargeIconPath=
 UATCommandLine=
 AutoSDKPath=
-TutorialPath=/Engine/Tutorial/Mobile/SettingUpLuminTutorial.SettingUpLuminTutorial
+TutorialPath=/Engine/Tutorial/Mobile/SettingUpLuminTutorial
 Windows:bIsEnabled=true
 Mac:bIsEnabled=true
 Linux:bIsEnabled=false
@@ -34,7 +34,7 @@ LargeIconPath=Launcher/Lumin/Platform_Lumin_128x
 XLargeIconPath=
 UATCommandLine=-client
 AutoSDKPath=
-TutorialPath=/Engine/Tutorial/Mobile/SettingUpLuminTutorial.SettingUpLuminTutorial
+TutorialPath=/Engine/Tutorial/Mobile/SettingUpLuminTutorial
 Windows:bIsEnabled=true
 Mac:bIsEnabled=true
 Linux:bIsEnabled=false
